Creating a Blank blueprint project without Starter Content now doesn't add 4 basic meshes to the project.

diff --git a/Templates/TP_BlankBP/Config/TemplateDefs.ini b/Templates/TP_BlankBP/Config/TemplateDefs.ini
--- a/Templates/TP_BlankBP/Config/TemplateDefs.ini
+++ b/Templates/TP_BlankBP/Config/TemplateDefs.ini
@@ -33,7 +33,3 @@ FilenameReplacements=(Extensions=("cpp","h","ini","cs"), From="%TEMPLATENAME%",
 ReplacementsInFiles=(Extensions=("cpp","h","ini","cs"), From="%TEMPLATENAME_UPPERCASE%", To="%PROJECTNAME_UPPERCASE%", bCaseSensitive=true)
 ReplacementsInFiles=(Extensions=("cpp","h","ini","cs"), From="%TEMPLATENAME_LOWERCASE%", To="%PROJECTNAME_LOWERCASE%", bCaseSensitive=true)
 ReplacementsInFiles=(Extensions=("cpp","h","ini","cs"), From="%TEMPLATENAME%", To="%PROJECTNAME%", bCaseSensitive=false)
-
-SharedContentPacks=(MountName="Geometry",DetailLevels=("Standard","High"))
-
-
