Added data-driven blank project templates.

diff --git a/Templates/TP_Blank/Config/TemplateDefs.ini b/Templates/TP_Blank/Config/TemplateDefs.ini
new file mode 100644
--- /dev/null
+++ b/Templates/TP_Blank/Config/TemplateDefs.ini
@@ -0,0 +1,30 @@
+[/Script/GameProjectGeneration.TemplateProjectDefs]
+LocalizedDisplayNames=(Language="en", Text="Blank")
+LocalizedDescriptions=(Language="en", Text="A clean empty project with no code.")
+
+FoldersToIgnore=Media
+
+FilesToIgnore="%TEMPLATENAME%.uproject"
+FilesToIgnore="%TEMPLATENAME%.png"
+FilesToIgnore="Config/TemplateDefs.ini"
+FilesToIgnore="Config/config.ini"
+FilesToIgnore="%TEMPLATENAME%.opensdf"
+FilesToIgnore="%TEMPLATENAME%.sdf"
+FilesToIgnore="%TEMPLATENAME%.v11.suo"
+FilesToIgnore="%TEMPLATENAME%.v12.suo"
+FilesToIgnore="%TEMPLATENAME%.sln"
+FilesToIgnore="Manifest.json"
+FilesToIgnore="contents.txt"
+
+FolderRenames=(From="Source/%TEMPLATENAME%", To="Source/%PROJECTNAME%")
+FolderRenames=(From="Source/%TEMPLATENAME%Editor", To="Source/%PROJECTNAME%Editor")
+
+FilenameReplacements=(Extensions=("cpp","h","ini","cs"), From="%TEMPLATENAME_UPPERCASE%", To="%PROJECTNAME_UPPERCASE%", bCaseSensitive=true)
+FilenameReplacements=(Extensions=("cpp","h","ini","cs"), From="%TEMPLATENAME_LOWERCASE%", To="%PROJECTNAME_LOWERCASE%", bCaseSensitive=true)
+FilenameReplacements=(Extensions=("cpp","h","ini","cs"), From="%TEMPLATENAME%",To="%PROJECTNAME%",bCaseSensitive=false)
+
+ReplacementsInFiles=(Extensions=("cpp","h","ini","cs"), From="%TEMPLATENAME_UPPERCASE%", To="%PROJECTNAME_UPPERCASE%", bCaseSensitive=true)
+ReplacementsInFiles=(Extensions=("cpp","h","ini","cs"), From="%TEMPLATENAME_LOWERCASE%", To="%PROJECTNAME_LOWERCASE%", bCaseSensitive=true)
+ReplacementsInFiles=(Extensions=("cpp","h","ini","cs"), From="%TEMPLATENAME%", To="%PROJECTNAME%", bCaseSensitive=false)
+
+SortKey="_1"
diff --git a/Templates/TP_Blank/Source/TP_Blank.Target.cs b/Templates/TP_Blank/Source/TP_Blank.Target.cs
new file mode 100644
--- /dev/null
+++ b/Templates/TP_Blank/Source/TP_Blank.Target.cs
@@ -0,0 +1,14 @@
+// Copyright 1998-2019 Epic Games, Inc. All Rights Reserved.
+
+using UnrealBuildTool;
+using System.Collections.Generic;
+
+public class TP_BlankTarget : TargetRules
+{
+	public TP_BlankTarget( TargetInfo Target) : base(Target)
+	{
+		Type = TargetType.Game;
+
+		ExtraModuleNames.AddRange( new string[] { "TP_Blank" } );
+	}
+}
diff --git a/Templates/TP_Blank/Source/TP_Blank/TP_Blank.Build.cs b/Templates/TP_Blank/Source/TP_Blank/TP_Blank.Build.cs
new file mode 100644
--- /dev/null
+++ b/Templates/TP_Blank/Source/TP_Blank/TP_Blank.Build.cs
@@ -0,0 +1,23 @@
+// Copyright 1998-2019 Epic Games, Inc. All Rights Reserved.
+
+using UnrealBuildTool;
+
+public class TP_Blank : ModuleRules
+{
+	public TP_Blank(ReadOnlyTargetRules Target) : base(Target)
+	{
+		PCHUsage = PCHUsageMode.UseExplicitOrSharedPCHs;
+	
+		PublicDependencyModuleNames.AddRange(new string[] { "Core", "CoreUObject", "Engine", "InputCore" });
+
+		PrivateDependencyModuleNames.AddRange(new string[] {  });
+
+		// Uncomment if you are using Slate UI
+		// PrivateDependencyModuleNames.AddRange(new string[] { "Slate", "SlateCore" });
+		
+		// Uncomment if you are using online features
+		// PrivateDependencyModuleNames.Add("OnlineSubsystem");
+
+		// To include OnlineSubsystemSteam, add it to the plugins section in your uproject file with the Enabled attribute set to true
+	}
+}
diff --git a/Templates/TP_Blank/Source/TP_Blank/TP_Blank.cpp b/Templates/TP_Blank/Source/TP_Blank/TP_Blank.cpp
new file mode 100644
--- /dev/null
+++ b/Templates/TP_Blank/Source/TP_Blank/TP_Blank.cpp
@@ -0,0 +1,6 @@
+// Copyright 1998-2019 Epic Games, Inc. All Rights Reserved.
+
+#include "TP_Blank.h"
+#include "Modules/ModuleManager.h"
+
+IMPLEMENT_PRIMARY_GAME_MODULE( FDefaultGameModuleImpl, TP_Blank, "TP_Blank" );
diff --git a/Templates/TP_Blank/Source/TP_Blank/TP_Blank.h b/Templates/TP_Blank/Source/TP_Blank/TP_Blank.h
new file mode 100644
--- /dev/null
+++ b/Templates/TP_Blank/Source/TP_Blank/TP_Blank.h
@@ -0,0 +1,6 @@
+// Copyright 1998-2019 Epic Games, Inc. All Rights Reserved.
+
+#pragma once
+
+#include "CoreMinimal.h"
+
diff --git a/Templates/TP_Blank/Source/TP_Blank/TP_BlankGameModeBase.cpp b/Templates/TP_Blank/Source/TP_Blank/TP_BlankGameModeBase.cpp
new file mode 100644
--- /dev/null
+++ b/Templates/TP_Blank/Source/TP_Blank/TP_BlankGameModeBase.cpp
@@ -0,0 +1,5 @@
+// Copyright 1998-2019 Epic Games, Inc. All Rights Reserved.
+
+
+#include "TP_BlankGameModeBase.h"
+
diff --git a/Templates/TP_Blank/Source/TP_Blank/TP_BlankGameModeBase.h b/Templates/TP_Blank/Source/TP_Blank/TP_BlankGameModeBase.h
new file mode 100644
--- /dev/null
+++ b/Templates/TP_Blank/Source/TP_Blank/TP_BlankGameModeBase.h
@@ -0,0 +1,17 @@
+// Copyright 1998-2019 Epic Games, Inc. All Rights Reserved.
+
+#pragma once
+
+#include "CoreMinimal.h"
+#include "GameFramework/GameModeBase.h"
+#include "TP_BlankGameModeBase.generated.h"
+
+/**
+ * 
+ */
+UCLASS()
+class TP_BLANK_API ATP_BlankGameModeBase : public AGameModeBase
+{
+	GENERATED_BODY()
+	
+};
diff --git a/Templates/TP_Blank/Source/TP_BlankEditor.Target.cs b/Templates/TP_Blank/Source/TP_BlankEditor.Target.cs
new file mode 100644
--- /dev/null
+++ b/Templates/TP_Blank/Source/TP_BlankEditor.Target.cs
@@ -0,0 +1,14 @@
+// Copyright 1998-2019 Epic Games, Inc. All Rights Reserved.
+
+using UnrealBuildTool;
+using System.Collections.Generic;
+
+public class TP_BlankEditorTarget : TargetRules
+{
+	public TP_BlankEditorTarget( TargetInfo Target) : base(Target)
+	{
+		Type = TargetType.Editor;
+
+		ExtraModuleNames.AddRange( new string[] { "TP_Blank" } );
+	}
+}
diff --git a/Templates/TP_Blank/TP_Blank.uproject b/Templates/TP_Blank/TP_Blank.uproject
new file mode 100644
--- /dev/null
+++ b/Templates/TP_Blank/TP_Blank.uproject
@@ -0,0 +1,7 @@
+{
+	"FileVersion": 3,
+	"EngineAssociation": "",
+	"Category": "",
+	"Description": "",
+	"Modules": []
+}
\ No newline at end of file
diff --git a/Templates/TP_Blank/contents.txt b/Templates/TP_Blank/contents.txt
new file mode 100644
--- /dev/null
+++ b/Templates/TP_Blank/contents.txt
@@ -0,0 +1,4 @@
+../../../Templates/TP_Blank/manifest.json
+../../../Templates/TP_Blank/Config/config.ini
+../../../Templates/TP_Blank/Media/TP_Blank.png
+../../../Templates/TP_Blank/Media/TP_Blank_Preview.png
diff --git a/Templates/TP_Blank/manifest.json b/Templates/TP_Blank/manifest.json
new file mode 100644
--- /dev/null
+++ b/Templates/TP_Blank/manifest.json
@@ -0,0 +1,29 @@
+{
+	"Version": 1,
+	"Name":
+	[
+		{
+			"Language": "en",
+			"Text": "Blank"
+		}
+	],
+	"Description":
+	[
+		{
+			"Language": "en",
+			"Text": "An empty project with some basic game framework code classes created."
+		}
+	],
+	"AssetTypes": [],	
+	"SearchTags": [],
+	"ClassTypes": "",
+	"Category": "",
+	"FocusAsset": "",
+	"Thumbnail": "TP_Blank.png",
+	"Screenshots":
+	[
+		"TP_Blank_Preview.png"
+	],
+	"SortKey": "_1",
+	"AdditionalFeaturePacks": []
+}
diff --git a/Templates/TP_BlankBP/Config/TemplateDefs.ini b/Templates/TP_BlankBP/Config/TemplateDefs.ini
new file mode 100644
--- /dev/null
+++ b/Templates/TP_BlankBP/Config/TemplateDefs.ini
@@ -0,0 +1,30 @@
+[/Script/GameProjectGeneration.TemplateProjectDefs]
+LocalizedDisplayNames=(Language="en", Text="Blank")
+LocalizedDescriptions=(Language="en", Text="A clean empty project with no code.")
+
+FoldersToIgnore=Media
+
+FilesToIgnore="%TEMPLATENAME%.uproject"
+FilesToIgnore="%TEMPLATENAME%.png"
+FilesToIgnore="Config/TemplateDefs.ini"
+FilesToIgnore="Config/config.ini"
+FilesToIgnore="%TEMPLATENAME%.opensdf"
+FilesToIgnore="%TEMPLATENAME%.sdf"
+FilesToIgnore="%TEMPLATENAME%.v11.suo"
+FilesToIgnore="%TEMPLATENAME%.v12.suo"
+FilesToIgnore="%TEMPLATENAME%.sln"
+FilesToIgnore="Manifest.json"
+FilesToIgnore="contents.txt"
+
+FolderRenames=(From="Source/%TEMPLATENAME%", To="Source/%PROJECTNAME%")
+FolderRenames=(From="Source/%TEMPLATENAME%Editor", To="Source/%PROJECTNAME%Editor")
+
+FilenameReplacements=(Extensions=("cpp","h","ini","cs"), From="%TEMPLATENAME_UPPERCASE%", To="%PROJECTNAME_UPPERCASE%", bCaseSensitive=true)
+FilenameReplacements=(Extensions=("cpp","h","ini","cs"), From="%TEMPLATENAME_LOWERCASE%", To="%PROJECTNAME_LOWERCASE%", bCaseSensitive=true)
+FilenameReplacements=(Extensions=("cpp","h","ini","cs"), From="%TEMPLATENAME%", To="%PROJECTNAME%", bCaseSensitive=false)
+
+ReplacementsInFiles=(Extensions=("cpp","h","ini","cs"), From="%TEMPLATENAME_UPPERCASE%", To="%PROJECTNAME_UPPERCASE%", bCaseSensitive=true)
+ReplacementsInFiles=(Extensions=("cpp","h","ini","cs"), From="%TEMPLATENAME_LOWERCASE%", To="%PROJECTNAME_LOWERCASE%", bCaseSensitive=true)
+ReplacementsInFiles=(Extensions=("cpp","h","ini","cs"), From="%TEMPLATENAME%", To="%PROJECTNAME%", bCaseSensitive=false)
+
+SortKey="_1"
diff --git a/Templates/TP_BlankBP/TP_BlankBP.uproject b/Templates/TP_BlankBP/TP_BlankBP.uproject
new file mode 100644
--- /dev/null
+++ b/Templates/TP_BlankBP/TP_BlankBP.uproject
@@ -0,0 +1,7 @@
+{
+	"FileVersion": 3,
+	"EngineAssociation": "",
+	"Category": "",
+	"Description": "",
+	"Modules": []
+}
\ No newline at end of file
diff --git a/Templates/TP_BlankBP/contents.txt b/Templates/TP_BlankBP/contents.txt
new file mode 100644
--- /dev/null
+++ b/Templates/TP_BlankBP/contents.txt
@@ -0,0 +1,4 @@
+../../../Templates/TP_BlankBP/manifest.json
+../../../Templates/TP_BlankBP/Config/config.ini
+../../../Templates/TP_BlankBP/Media/TP_BlankBP.png
+../../../Templates/TP_BlankBP/Media/TP_BlankBP_Preview.png
diff --git a/Templates/TP_BlankBP/manifest.json b/Templates/TP_BlankBP/manifest.json
new file mode 100644
--- /dev/null
+++ b/Templates/TP_BlankBP/manifest.json
@@ -0,0 +1,29 @@
+{
+	"Version": 1,
+	"Name":
+	[
+		{
+			"Language": "en",
+			"Text": "Blank"
+		}
+	],
+	"Description":
+	[
+		{
+			"Language": "en",
+			"Text": "A clean empty project with no code."
+		}
+	],
+	"AssetTypes": [],
+	"SearchTags": [],
+	"ClassTypes": "",
+	"Category": "",
+	"FocusAsset": "",
+	"Thumbnail": "TP_BlankBP.png",
+	"Screenshots":
+	[
+		"TP_BlankBP_Preview.png"
+	],
+	"SortKey": "_1",
+	"AdditionalFeaturePacks": []
+}
