Blank templates no longer import the Geometry module.

diff --git a/Templates/TP_Blank/Config/TemplateDefs.ini b/Templates/TP_Blank/Config/TemplateDefs.ini
--- a/Templates/TP_Blank/Config/TemplateDefs.ini
+++ b/Templates/TP_Blank/Config/TemplateDefs.ini
@@ -105,3 +105,7 @@ ReplacementsInFiles=(Extensions=("cpp","h","ini","cs"), From="%TEMPLATENAME%", T
 
 ; Shared feature packs. The files in these packs listed in these structures marked as 'additionalfiles' will be copied on project generation.
 ;SharedContentPacks=(MountName="2DSideScroller",DetailLevels=("Standard"))
+
+; Starter content feature pack to use if the user chooses to add starter content.
+; This defaults to StarterContent on desktop and MobileStarterContent on mobile.
+;StarterContent=TP_FirstPerson
\ No newline at end of file
diff --git a/Templates/TP_Blank/manifest.json b/Templates/TP_Blank/manifest.json
--- a/Templates/TP_Blank/manifest.json
+++ b/Templates/TP_Blank/manifest.json
@@ -25,15 +25,5 @@
 		"TP_Blank_Preview.png"
 	],
 	"SortKey": "_1",
-	"AdditionalFeaturePacks":
-	[
-		{
-			"MountName": "Geometry",
-			"DetailLevels":
-			[
-				"Standard",
-				"High"
-			]
-		}
-	]
+	"AdditionalFeaturePacks": []
 }
diff --git a/Templates/TP_BlankBP/Config/TemplateDefs.ini b/Templates/TP_BlankBP/Config/TemplateDefs.ini
--- a/Templates/TP_BlankBP/Config/TemplateDefs.ini
+++ b/Templates/TP_BlankBP/Config/TemplateDefs.ini
@@ -33,7 +33,3 @@ FilenameReplacements=(Extensions=("cpp","h","ini","cs"), From="%TEMPLATENAME%",
 ReplacementsInFiles=(Extensions=("cpp","h","ini","cs"), From="%TEMPLATENAME_UPPERCASE%", To="%PROJECTNAME_UPPERCASE%", bCaseSensitive=true)
 ReplacementsInFiles=(Extensions=("cpp","h","ini","cs"), From="%TEMPLATENAME_LOWERCASE%", To="%PROJECTNAME_LOWERCASE%", bCaseSensitive=true)
 ReplacementsInFiles=(Extensions=("cpp","h","ini","cs"), From="%TEMPLATENAME%", To="%PROJECTNAME%", bCaseSensitive=false)
-
-SharedContentPacks=(MountName="Geometry",DetailLevels=("Standard","High"))
-
-
diff --git a/Templates/TP_BlankBP/manifest.json b/Templates/TP_BlankBP/manifest.json
--- a/Templates/TP_BlankBP/manifest.json
+++ b/Templates/TP_BlankBP/manifest.json
@@ -25,15 +25,5 @@
 		"TP_BlankBP_Preview.png"
 	],
 	"SortKey": "_1",
-	"AdditionalFeaturePacks":
-	[
-		{
-			"MountName": "Geometry",
-			"DetailLevels":
-			[
-				"Standard",
-				"High"
-			]
-		}
-	]
+	"AdditionalFeaturePacks": []
 }
