adding Chinese translation to this file #jira ueloc-4711 #fyi jamie.dale, michelle.zhao, alexis.zhang, yali.he

diff --git a/Templates/TP_BlankBP/Config/TemplateDefs.ini b/Templates/TP_BlankBP/Config/TemplateDefs.ini
--- a/Templates/TP_BlankBP/Config/TemplateDefs.ini
+++ b/Templates/TP_BlankBP/Config/TemplateDefs.ini
@@ -1,4 +1,4 @@
-[/Script/GameProjectGeneration.TemplateProjectDefs]
+﻿[/Script/GameProjectGeneration.TemplateProjectDefs]
 
 SortKey="_1"
 
@@ -8,6 +8,8 @@ Categories=Games
 
 LocalizedDisplayNames=(Language="en", Text="Blank")
 LocalizedDescriptions=(Language="en", Text="A clean empty project with no code.")
+LocalizedDisplayNames=(Language="zh-Hans", Text="空白")
+LocalizedDescriptions=(Language="zh-Hans", Text="不含任何代码的空白项目。")
 
 FoldersToIgnore=Media
 
