Updating JA translation for "Blank"

diff --git a/Templates/TP_Blank/Config/TemplateDefs.ini b/Templates/TP_Blank/Config/TemplateDefs.ini
--- a/Templates/TP_Blank/Config/TemplateDefs.ini
+++ b/Templates/TP_Blank/Config/TemplateDefs.ini
@@ -11,6 +11,7 @@ Categories=Games
 LocalizedDisplayNames=(Language="en", Text="Blank")
 LocalizedDisplayNames=(Language="zh-Hans", Text="空白")
 LocalizedDisplayNames=(Language="ko", Text="기본")
+LocalizedDisplayNames=(Language="ja", Text="ブランク")
 
 ; The localized descriptions for this template, which is displayed when the template is selected.
 LocalizedDescriptions=(Language="en", Text="A clean empty project with no code.")
diff --git a/Templates/TP_BlankBP/Config/TemplateDefs.ini b/Templates/TP_BlankBP/Config/TemplateDefs.ini
--- a/Templates/TP_BlankBP/Config/TemplateDefs.ini
+++ b/Templates/TP_BlankBP/Config/TemplateDefs.ini
@@ -10,6 +10,7 @@ LocalizedDisplayNames=(Language="en", Text="Blank")
 LocalizedDescriptions=(Language="en", Text="A clean empty project with no code.")
 LocalizedDisplayNames=(Language="zh-Hans", Text="空白")
 LocalizedDescriptions=(Language="zh-Hans", Text="不含任何代码的空白项目。")
+LocalizedDisplayNames=(Language="ja", Text="ブランク")
 LocalizedDescriptions=(Language="ja", Text="まったく空のプロジェクトです。コードは入っていません。")
 LocalizedDisplayNames=(Language="ko", Text="기본")
 LocalizedDescriptions=(Language="ko", Text="코드가 포함되지 않은 기본 빈 프로젝트입니다.")
